Added Simulation templates :
- Blank
- Sandbox

Associated the HandheldAR and VirtualReality templates to the Simulation category

diff --git a/Templates/TP_HandheldARBP/Config/TemplateDefs.ini b/Templates/TP_HandheldARBP/Config/TemplateDefs.ini
--- a/Templates/TP_HandheldARBP/Config/TemplateDefs.ini
+++ b/Templates/TP_HandheldARBP/Config/TemplateDefs.ini
@@ -1,6 +1,8 @@
 ﻿[/Script/GameProjectGeneration.TemplateProjectDefs]
 
 Categories=Games
+Categories=Sim
+
 
 LocalizedDisplayNames=(Language="en",Text="Handheld AR")
 LocalizedDescriptions=(Language="en",Text="A starting point for building Augmented Reality applications for Android and iOS devices. Includes runtime logic for toggling AR mode on and off, debug information regarding plane detection, example code for hit detection and handling light estimation.")
diff --git a/Templates/TP_SIM_Blank/Config/DefaultEngine.ini b/Templates/TP_SIM_Blank/Config/DefaultEngine.ini
--- a/Templates/TP_SIM_Blank/Config/DefaultEngine.ini
+++ b/Templates/TP_SIM_Blank/Config/DefaultEngine.ini
@@ -1,8 +1,8 @@
 
 
 [/Script/EngineSettings.GameMapsSettings]
-EditorStartupMap=/Game/SimDefault/Levels/SimDefault
-+EditorTemplateMapOverrides=(Thumbnail="/Game/SimDefault/Levels/Thumbnails/SimDefault_Thumbnail.SimDefault_Thumbnail",Map="/Game/SimDefault/Levels/SimDefault.SimDefault",DisplayName=NSLOCTEXT("[/Script/EngineSettings]", "C47C7345459213852543F1906820C22C", "Simulation Default"))
+EditorStartupMap=/Game/SimBlank/Levels/SimBlank
++EditorTemplateMapOverrides=(Thumbnail="/Game/SimBlank/Levels/Thumbnails/SimBlank_Thumbnail.SimBlank_Thumbnail",Map="/Game/SimBlank/Levels/SimBlank.SimBlank",DisplayName=NSLOCTEXT("[/Script/EngineSettings]", "C47C7345459213852543F1906820C22C", "Simulation Default"))
 GameDefaultMap=/Game/Main
 
 [/Script/Engine.RendererSettings]
@@ -10,6 +10,12 @@ r.DefaultFeature.AutoExposure.ExtendDefaultLuminanceRange=True
 r.AllowStaticLighting=False
 r.Shadow.Virtual.Enable=1
 
+r.GenerateMeshDistanceFields=True
+
+r.DynamicGlobalIlluminationMethod=1
+
+r.ReflectionMethod=1
+
 [/Script/AndroidFileServerEditor.AndroidFileServerRuntimeSettings]
 bEnablePlugin=True
 bAllowNetworkConnection=True
@@ -52,4 +58,3 @@ MinSampleRate=8000.000000
 CompressionQualityModifier=1.000000
 AutoStreamingThreshold=0.000000
 SoundCueCookQualityIndex=-1
-
diff --git a/Templates/TP_SIM_Blank/Config/DefaultGame.ini b/Templates/TP_SIM_Blank/Config/DefaultGame.ini
--- a/Templates/TP_SIM_Blank/Config/DefaultGame.ini
+++ b/Templates/TP_SIM_Blank/Config/DefaultGame.ini
@@ -1,3 +1,3 @@
 
 [/Script/EngineSettings.GeneralProjectSettings]
-ProjectName=Blank
\ No newline at end of file
+ProjectName=SimBlank
\ No newline at end of file
diff --git a/Templates/TP_SIM_Blank/Config/TemplateDefs.ini b/Templates/TP_SIM_Blank/Config/TemplateDefs.ini
--- a/Templates/TP_SIM_Blank/Config/TemplateDefs.ini
+++ b/Templates/TP_SIM_Blank/Config/TemplateDefs.ini
@@ -13,7 +13,7 @@ LocalizedDisplayNames=(Language="zh-Hans", Text="空白")
 LocalizedDisplayNames=(Language="ko", Text="기본")
 
 ; The localized descriptions for this template, which is displayed when the template is selected.
-LocalizedDescriptions=(Language="en", Text="Blank template with minimal functionality and content.\n\nThis template enables you to create new outdoor scenes with appropriate lighting and georeferencing.")
+LocalizedDescriptions=(Language="en", Text="A basic empty environment with defaults settings suitable for Simulation projects.\n\nEarth Atmosphere, Lighting, Georeferencing...")
 LocalizedDescriptions=(Language="zh-Hans", Text="不含任何代码的空白项目。")
 LocalizedDescriptions=(Language="ja", Text="まったく空のプロジェクトです。コードは入っていません。")
 LocalizedDescriptions=(Language="ko", Text="코드가 포함되지 않은 기본 빈 프로젝트입니다.")
@@ -53,17 +53,58 @@ FilesToIgnore="Manifest.json"
 FilesToIgnore="contents.txt"
 
 
-FolderRenames=(From="Source/%TEMPLATENAME%",To="Source/%PROJECTNAME%")
-FolderRenames=(From="Source/%TEMPLATENAME%Editor",To="Source/%PROJECTNAME%Editor")
+; Folders that will be renamed when creating a project from this template.
+; %PROJECTNAME% will be replaced by the project name specified by the user during creation when resolving these rules.
+; For example,  then "Source/TP_Blank" will be copied to "Source/MyProject".
+FolderRenames=(From="Source/%TEMPLATENAME%", To="Source/%PROJECTNAME%")
+FolderRenames=(From="Source/%TEMPLATENAME%Editor", To="Source/%PROJECTNAME%Editor")
 
-FilenameReplacements=(Extensions=("cpp","h","ini","cs"),From="%TEMPLATENAME_UPPERCASE%",To="%PROJECTNAME_UPPERCASE%",bCaseSensitive=true)
-FilenameReplacements=(Extensions=("cpp","h","ini","cs"),From="%TEMPLATENAME_LOWERCASE%",To="%PROJECTNAME_LOWERCASE%",bCaseSensitive=true)
-FilenameReplacements=(Extensions=("cpp","h","ini","cs"),From="%TEMPLATENAME%",To="%PROJECTNAME%",bCaseSensitive=false)
+; Rules for renaming files when creating a project from this template. 
+; %TEMPLATENAME_UPPERCASE% and %PROJECTNAME_UPPERCASE% are the template and project names in upper case, eg. for "TP_Blank", its upper case representation is "TP_BLANK".
+; %TEMPLATENAME_LOWERCASE% and %PROJECTNAME_LOWERCASE% are the template and project names in lower case, eg. for "TP_Blank", its lower case representation is "tp_blank".
+; A rule does not have to include any of these modifiers. 
+; Extensions is a list of files to consider for these replacements. Eg. ("cpp", "h") will try to apply a rule only to files like "SomeFile.cpp" and "SomeFile.h".
+; bCaseSensitive specifies whether this replacement should be case-sensitive.
 
-ReplacementsInFiles=(Extensions=("cpp","h","ini","cs"),From="%TEMPLATENAME_UPPERCASE%",To="%PROJECTNAME_UPPERCASE%",bCaseSensitive=true)
-ReplacementsInFiles=(Extensions=("cpp","h","ini","cs"),From="%TEMPLATENAME_LOWERCASE%",To="%PROJECTNAME_LOWERCASE%",bCaseSensitive=true)
-ReplacementsInFiles=(Extensions=("cpp","h","ini","cs"),From="%TEMPLATENAME%",To="%PROJECTNAME%",bCaseSensitive=false)
+; Here are some default rules for renaming files:
+
+; This rule states that a file called "TP_BLANK.cpp" will be renamed to "MYPROJECT.cpp".
+FilenameReplacements=(Extensions=("cpp","h","ini","cs"), From="%TEMPLATENAME_UPPERCASE%", To="%PROJECTNAME_UPPERCASE%", bCaseSensitive=true)
+
+; This rule states that a file called "tp_blank.ini" will be renamed to "myproject.ini"
+FilenameReplacements=(Extensions=("cpp","h","ini","cs"), From="%TEMPLATENAME_LOWERCASE%", To="%PROJECTNAME_LOWERCASE%", bCaseSensitive=true)
+
+; And this rule catches any irregularly-cased files, so that "Tp_bLANK.ini" will be renamed to "MyProject.ini".
+FilenameReplacements=(Extensions=("cpp","h","ini","cs"), From="%TEMPLATENAME%",To="%PROJECTNAME%",bCaseSensitive=false)
+
+; Rules for renaming strings in the contents of files. 
+; The same rules as above apply, but any usages of strings that match these rules inside the file will be changed. This includes all substrings!
+; For example, "TP_Blank" will be replaced by "MyProject", but "TP_Blank_Name" will also be replaced by "MyProject_Name".
+ReplacementsInFiles=(Extensions=("cpp","h","ini","cs"), From="%TEMPLATENAME_UPPERCASE%", To="%PROJECTNAME_UPPERCASE%", bCaseSensitive=true)
+ReplacementsInFiles=(Extensions=("cpp","h","ini","cs"), From="%TEMPLATENAME_LOWERCASE%", To="%PROJECTNAME_LOWERCASE%", bCaseSensitive=true)
+ReplacementsInFiles=(Extensions=("cpp","h","ini","cs"), From="%TEMPLATENAME%", To="%PROJECTNAME%", bCaseSensitive=false)
+
+; The following settings are not enabled for this project, but are included for completeness:
+
+; Whether or not this template can be used to create a project. Default: true
+; If this is set to false, then the template becomes a Feature Pack that can only be referenced from other templates, or added to a project from the Content Browser.
+;bAllowProjectCreation=false
+
+; Which settings on the project settings page should be hidden for this template. Default: none
+; Allowed values are:
+; All - Hides the settings page entirely
+; Languages - Choose between C++ and Blueprint versions.
+; HardwareTarget - Choose between Mobile and Desktop hardware targets.
+; GraphicsPreset - Choose between Maximum and Scalable graphics.
+; StarterContent - Choose whether to add starter content to the project.
+; XR - Choose whether to enable Virtual Reality support for the project.
+;HiddenSettings=All
+;HiddenSettings=Languages
+HiddenSettings=HardwareTarget
+HiddenSettings=GraphicsPreset
+HiddenSettings=RayTracing
+;HiddenSettings=StarterContent 
+HiddenSettings=XR
 
 ; Shared feature packs. The files in these packs listed in these structures marked as 'additionalfiles' will be copied on project generation.
-;SharedContentPacks=(MountName="2DSideScroller",DetailLevels=("Standard"))
-SharedContentPacks=(MountName="SimDefault",DetailLevels=("Standard"))
\ No newline at end of file
+SharedContentPacks=(MountName="SimBlank",DetailLevels=("Standard"))
diff --git a/Templates/TP_SIM_Blank/manifest.json b/Templates/TP_SIM_Blank/manifest.json
new file mode 100644
--- /dev/null
+++ b/Templates/TP_SIM_Blank/manifest.json
@@ -0,0 +1,73 @@
+{
+	"Version": 1,
+	"Name":
+	[
+		{
+			"Language": "en",
+			"Text": "Simulation Blank"
+		},
+		{
+			"Language": "ko",
+			"Text": "일인칭"
+		},
+		{
+			"Language": "zh-Hans",
+			"Text": "第一人称游戏"
+		},
+		{
+			"Language": "ja",
+			"Text": "ファーストパーソン"
+		}
+	],
+	"Description":
+	[
+		{
+			"Language": "en",
+			"Text": "A basic empty environment with defauts settings suitable for simulation projects."
+		},
+		{
+			"Language": "ko",
+			"Text": "일인칭 템플릿은 일인칭 시점에서 보이는 (양팔로 표현되는) 플레이어 캐릭터를 특징으로 합니다. 캐릭터 이동은 키보드, 컨트롤러, 터치 디바이스의 가상 조이스틱으로 가능합니다. 추가적으로 플레이어는 마우스, 컨트롤러, 터치 디바이스의 가상 조이스틱을 사용하여 주변을 둘러볼 수도 있습니다. 캐릭터에는 총이 장착되어 있어 마우스, 컨트롤러, 터치 디바이스의 가상 조이스틱을 사용하면 간단한 발사체가 발사되어, 아레나 벽에 튕기면서 레벨에 있는 몇몇 물리 오브젝트에 영향을 줍니다."
+		},
+		{
+			"Language": "zh-Hans",
+			"Text": "第一人称模板包含一个第一人称玩家角色（通过一双手臂呈现）。你可以通过键盘、手柄或触摸式设备上的虚拟摇杆控制角色在关卡中四处移动。你还可以使用鼠标、手柄或触摸式设备上的虚拟摇杆转动角色视角，观察四周。此外，你还可以让角色拾取枪支，然后用鼠标、手柄或触摸式设备上的虚拟摇杆来发射球形子弹。子弹会撞飞关卡中的物理对象，还能从竞技场墙壁上弹开。"
+		},
+		{
+			"Language": "ja",
+			"Text": "ファーストパーソンテンプレートには、一人称視点で表示されるプレイヤーキャラクタ（腕として表示されています）がいます。キーボード、コントローラー、タッチデバイスでのバーチャルジョイスティックでこのキャラクタを移動させることができます。追加して、プレイヤーはマウス、コントローラー、タッチデバイスでのバーチャルジョイスティックで周りを見回すことができます。さらにキャラクタは銃を装備して、マウス、コントローラー、タッチデバイスでのバーチャルジョイスティックを使ってシンプルな球体の弾を発射できます。この弾は物理オブジェクトに影響し、アリーナの壁で反射します。"
+		}
+	],
+	"AssetTypes":
+	[
+		{
+			"Language": "en",
+			"Text": ""
+		}
+	],	
+	"SearchTags":
+	[
+		{
+			"Language": "en",
+			"Text": "Simulation"
+		}
+	],
+	"ClassTypes": "Actor",
+	"Category": "BlueprintFeature",
+	"FocusAsset": "",
+	"Thumbnail": "TP_SIM_Blank.png",
+	"Screenshots":
+	[
+		"TP_SIM_Blank_Preview.png"
+	],
+	"AdditionalFeaturePacks":
+	[
+		{
+			"MountName": "SimBlank",
+			"DetailLevels":
+			[
+				"Standard"
+			]
+		}
+	]
+}
diff --git a/Templates/TP_SIM_BlankBP/Config/DefaultEngine.ini b/Templates/TP_SIM_BlankBP/Config/DefaultEngine.ini
--- a/Templates/TP_SIM_BlankBP/Config/DefaultEngine.ini
+++ b/Templates/TP_SIM_BlankBP/Config/DefaultEngine.ini
@@ -1,8 +1,8 @@
 
 
 [/Script/EngineSettings.GameMapsSettings]
-EditorStartupMap=/Game/SimDefault/Levels/SimDefault
-+EditorTemplateMapOverrides=(Thumbnail="/Game/SimDefault/Levels/Thumbnails/SimDefault_Thumbnail.SimDefault_Thumbnail",Map="/Game/SimDefault/Levels/SimDefault.SimDefault",DisplayName=NSLOCTEXT("[/Script/EngineSettings]", "C47C7345459213852543F1906820C22C", "Simulation Default"))
+EditorStartupMap=/Game/SimBlank/Levels/SimBlank
++EditorTemplateMapOverrides=(Thumbnail="/Game/SimBlank/Levels/Thumbnails/SimBlank_Thumbnail.SimBlank_Thumbnail",Map="/Game/SimBlank/Levels/SimBlank.SimBlank",DisplayName=NSLOCTEXT("[/Script/EngineSettings]", "C47C7345459213852543F1906820C22C", "Simulation Default"))
 GameDefaultMap=/Game/Main
 
 [/Script/Engine.RendererSettings]
@@ -10,6 +10,12 @@ r.DefaultFeature.AutoExposure.ExtendDefaultLuminanceRange=True
 r.AllowStaticLighting=False
 r.Shadow.Virtual.Enable=1
 
+r.GenerateMeshDistanceFields=True
+
+r.DynamicGlobalIlluminationMethod=1
+
+r.ReflectionMethod=1
+
 [/Script/AndroidFileServerEditor.AndroidFileServerRuntimeSettings]
 bEnablePlugin=True
 bAllowNetworkConnection=True
@@ -52,4 +58,3 @@ MinSampleRate=8000.000000
 CompressionQualityModifier=1.000000
 AutoStreamingThreshold=0.000000
 SoundCueCookQualityIndex=-1
-
diff --git a/Templates/TP_SIM_BlankBP/Config/DefaultGame.ini b/Templates/TP_SIM_BlankBP/Config/DefaultGame.ini
--- a/Templates/TP_SIM_BlankBP/Config/DefaultGame.ini
+++ b/Templates/TP_SIM_BlankBP/Config/DefaultGame.ini
@@ -1,3 +1,3 @@
 
 [/Script/EngineSettings.GeneralProjectSettings]
-ProjectName=Blank
\ No newline at end of file
+ProjectName=TP_SimBlank
\ No newline at end of file
diff --git a/Templates/TP_SIM_BlankBP/Config/TemplateDefs.ini b/Templates/TP_SIM_BlankBP/Config/TemplateDefs.ini
--- a/Templates/TP_SIM_BlankBP/Config/TemplateDefs.ini
+++ b/Templates/TP_SIM_BlankBP/Config/TemplateDefs.ini
@@ -13,7 +13,7 @@ LocalizedDisplayNames=(Language="zh-Hans", Text="空白")
 LocalizedDisplayNames=(Language="ko", Text="기본")
 
 ; The localized descriptions for this template, which is displayed when the template is selected.
-LocalizedDescriptions=(Language="en", Text="Blank template with minimal functionality and content.\n\nThis template enables you to create new outdoor scenes with appropriate lighting and georeferencing.")
+LocalizedDescriptions=(Language="en", Text="A basic empty environment with defaults settings suitable for Simulation projects.\n\nEarth Atmosphere, Lighting, Georeferencing...")
 LocalizedDescriptions=(Language="zh-Hans", Text="不含任何代码的空白项目。")
 LocalizedDescriptions=(Language="ja", Text="まったく空のプロジェクトです。コードは入っていません。")
 LocalizedDescriptions=(Language="ko", Text="코드가 포함되지 않은 기본 빈 프로젝트입니다.")
@@ -53,17 +53,58 @@ FilesToIgnore="Manifest.json"
 FilesToIgnore="contents.txt"
 
 
-FolderRenames=(From="Source/%TEMPLATENAME%",To="Source/%PROJECTNAME%")
-FolderRenames=(From="Source/%TEMPLATENAME%Editor",To="Source/%PROJECTNAME%Editor")
+; Folders that will be renamed when creating a project from this template.
+; %PROJECTNAME% will be replaced by the project name specified by the user during creation when resolving these rules.
+; For example,  then "Source/TP_Blank" will be copied to "Source/MyProject".
+FolderRenames=(From="Source/%TEMPLATENAME%", To="Source/%PROJECTNAME%")
+FolderRenames=(From="Source/%TEMPLATENAME%Editor", To="Source/%PROJECTNAME%Editor")
 
-FilenameReplacements=(Extensions=("cpp","h","ini","cs"),From="%TEMPLATENAME_UPPERCASE%",To="%PROJECTNAME_UPPERCASE%",bCaseSensitive=true)
-FilenameReplacements=(Extensions=("cpp","h","ini","cs"),From="%TEMPLATENAME_LOWERCASE%",To="%PROJECTNAME_LOWERCASE%",bCaseSensitive=true)
-FilenameReplacements=(Extensions=("cpp","h","ini","cs"),From="%TEMPLATENAME%",To="%PROJECTNAME%",bCaseSensitive=false)
+; Rules for renaming files when creating a project from this template. 
+; %TEMPLATENAME_UPPERCASE% and %PROJECTNAME_UPPERCASE% are the template and project names in upper case, eg. for "TP_Blank", its upper case representation is "TP_BLANK".
+; %TEMPLATENAME_LOWERCASE% and %PROJECTNAME_LOWERCASE% are the template and project names in lower case, eg. for "TP_Blank", its lower case representation is "tp_blank".
+; A rule does not have to include any of these modifiers. 
+; Extensions is a list of files to consider for these replacements. Eg. ("cpp", "h") will try to apply a rule only to files like "SomeFile.cpp" and "SomeFile.h".
+; bCaseSensitive specifies whether this replacement should be case-sensitive.
 
-ReplacementsInFiles=(Extensions=("cpp","h","ini","cs"),From="%TEMPLATENAME_UPPERCASE%",To="%PROJECTNAME_UPPERCASE%",bCaseSensitive=true)
-ReplacementsInFiles=(Extensions=("cpp","h","ini","cs"),From="%TEMPLATENAME_LOWERCASE%",To="%PROJECTNAME_LOWERCASE%",bCaseSensitive=true)
-ReplacementsInFiles=(Extensions=("cpp","h","ini","cs"),From="%TEMPLATENAME%",To="%PROJECTNAME%",bCaseSensitive=false)
+; Here are some default rules for renaming files:
+
+; This rule states that a file called "TP_BLANK.cpp" will be renamed to "MYPROJECT.cpp".
+FilenameReplacements=(Extensions=("cpp","h","ini","cs"), From="%TEMPLATENAME_UPPERCASE%", To="%PROJECTNAME_UPPERCASE%", bCaseSensitive=true)
+
+; This rule states that a file called "tp_blank.ini" will be renamed to "myproject.ini"
+FilenameReplacements=(Extensions=("cpp","h","ini","cs"), From="%TEMPLATENAME_LOWERCASE%", To="%PROJECTNAME_LOWERCASE%", bCaseSensitive=true)
+
+; And this rule catches any irregularly-cased files, so that "Tp_bLANK.ini" will be renamed to "MyProject.ini".
+FilenameReplacements=(Extensions=("cpp","h","ini","cs"), From="%TEMPLATENAME%",To="%PROJECTNAME%",bCaseSensitive=false)
+
+; Rules for renaming strings in the contents of files. 
+; The same rules as above apply, but any usages of strings that match these rules inside the file will be changed. This includes all substrings!
+; For example, "TP_Blank" will be replaced by "MyProject", but "TP_Blank_Name" will also be replaced by "MyProject_Name".
+ReplacementsInFiles=(Extensions=("cpp","h","ini","cs"), From="%TEMPLATENAME_UPPERCASE%", To="%PROJECTNAME_UPPERCASE%", bCaseSensitive=true)
+ReplacementsInFiles=(Extensions=("cpp","h","ini","cs"), From="%TEMPLATENAME_LOWERCASE%", To="%PROJECTNAME_LOWERCASE%", bCaseSensitive=true)
+ReplacementsInFiles=(Extensions=("cpp","h","ini","cs"), From="%TEMPLATENAME%", To="%PROJECTNAME%", bCaseSensitive=false)
+
+; The following settings are not enabled for this project, but are included for completeness:
+
+; Whether or not this template can be used to create a project. Default: true
+; If this is set to false, then the template becomes a Feature Pack that can only be referenced from other templates, or added to a project from the Content Browser.
+;bAllowProjectCreation=false
+
+; Which settings on the project settings page should be hidden for this template. Default: none
+; Allowed values are:
+; All - Hides the settings page entirely
+; Languages - Choose between C++ and Blueprint versions.
+; HardwareTarget - Choose between Mobile and Desktop hardware targets.
+; GraphicsPreset - Choose between Maximum and Scalable graphics.
+; StarterContent - Choose whether to add starter content to the project.
+; XR - Choose whether to enable Virtual Reality support for the project.
+;HiddenSettings=All
+;HiddenSettings=Languages
+HiddenSettings=HardwareTarget
+HiddenSettings=GraphicsPreset
+HiddenSettings=RayTracing
+;HiddenSettings=StarterContent 
+HiddenSettings=XR
 
 ; Shared feature packs. The files in these packs listed in these structures marked as 'additionalfiles' will be copied on project generation.
-;SharedContentPacks=(MountName="2DSideScroller",DetailLevels=("Standard"))
-SharedContentPacks=(MountName="SimDefault",DetailLevels=("Standard"))
\ No newline at end of file
+SharedContentPacks=(MountName="SimBlank",DetailLevels=("Standard"))
diff --git a/Templates/TP_SIM_BlankBP/manifest.json b/Templates/TP_SIM_BlankBP/manifest.json
new file mode 100644
--- /dev/null
+++ b/Templates/TP_SIM_BlankBP/manifest.json
@@ -0,0 +1,72 @@
+{
+	"Version": 1,
+	"Name":
+	[
+		{
+			"Language": "en",
+			"Text": "Simulation Blank"
+		},
+		{
+			"Language": "ko",
+			"Text": "일인칭"
+		},
+		{
+			"Language": "zh-Hans",
+			"Text": "第一人称游戏"
+		},
+		{
+			"Language": "ja",
+			"Text": "ファーストパーソン"
+		}
+	],
+	"Description":
+	[
+		{
+			"Language": "en",
+			"Text": "A basic empty environment with defauts settings suitable for simulation projects."
+		},
+		{
+			"Language": "ko",
+			"Text": "일인칭 템플릿은 일인칭 시점에서 보이는 (양팔로 표현되는) 플레이어 캐릭터를 특징으로 합니다. 캐릭터 이동은 키보드, 컨트롤러, 터치 디바이스의 가상 조이스틱으로 가능합니다. 추가적으로 플레이어는 마우스, 컨트롤러, 터치 디바이스의 가상 조이스틱을 사용하여 주변을 둘러볼 수도 있습니다. 캐릭터에는 총이 장착되어 있어 마우스, 컨트롤러, 터치 디바이스의 가상 조이스틱을 사용하면 간단한 발사체가 발사되어, 아레나 벽에 튕기면서 레벨에 있는 몇몇 물리 오브젝트에 영향을 줍니다."
+		},
+		{
+			"Language": "zh-Hans",
+			"Text": "第一人称模板包含一个第一人称玩家角色（通过一双手臂呈现）。你可以通过键盘、手柄或触摸式设备上的虚拟摇杆控制角色在关卡中四处移动。你还可以使用鼠标、手柄或触摸式设备上的虚拟摇杆转动角色视角，观察四周。此外，你还可以让角色拾取枪支，然后用鼠标、手柄或触摸式设备上的虚拟摇杆来发射球形子弹。子弹会撞飞关卡中的物理对象，还能从竞技场墙壁上弹开。"
+		},
+		{
+			"Language": "ja",
+			"Text": "ファーストパーソンテンプレートには、一人称視点で表示されるプレイヤーキャラクタ（腕として表示されています）がいます。キーボード、コントローラー、タッチデバイスでのバーチャルジョイスティックでこのキャラクタを移動させることができます。追加して、プレイヤーはマウス、コントローラー、タッチデバイスでのバーチャルジョイスティックで周りを見回すことができます。さらにキャラクタは銃を装備して、マウス、コントローラー、タッチデバイスでのバーチャルジョイスティックを使ってシンプルな球体の弾を発射できます。この弾は物理オブジェクトに影響し、アリーナの壁で反射します。"
+		}
+	],
+	"AssetTypes":
+	[
+		{
+			"Language": "en",
+			"Text": ""
+		}
+	],	
+	"SearchTags":
+	[
+		{
+			"Language": "en",
+			"Text": "Simulation"
+		}
+	],
+	"ClassTypes": "Actor",
+	"Category": "BlueprintFeature",
+	"FocusAsset": "",
+	"Thumbnail": "TP_SIM_BlankBP.png",
+	"Screenshots":
+	[
+		"TP_SIM_BlankBP_Preview.png"
+	],
+	"AdditionalFeaturePacks": [
+		{
+			"MountName": "SimBlank",
+			"DetailLevels":
+			[
+				"Standard"
+			]
+		}
+	]
+}
diff --git a/Templates/TP_SIM_Sandbox/Config/DefaultEngine.ini b/Templates/TP_SIM_Sandbox/Config/DefaultEngine.ini
new file mode 100644
--- /dev/null
+++ b/Templates/TP_SIM_Sandbox/Config/DefaultEngine.ini
@@ -0,0 +1,61 @@
+
+
+[/Script/EngineSettings.GameMapsSettings]
+EditorStartupMap=/Game/SimSandbox/Levels/SimSandbox
++EditorTemplateMapOverrides=(Thumbnail="/Game/SimSandbox/Levels/Thumbnails/SimSandbox_Thumbnail.SimSandbox_Thumbnail",Map="/Game/SimSandbox/Levels/SimSandbox.SimSandbox",DisplayName=NSLOCTEXT("[/Script/EngineSettings]", "C47C7345459213852543F1906820C22C", "Simulation Sandbox"))
++EditorTemplateMapOverrides=(Thumbnail="/Game/SimBlank/Levels/Thumbnails/SimBlank_Thumbnail.SimBlank_Thumbnail",Map="/Game/SimBlank/Levels/SimBlank.SimBlank",DisplayName=NSLOCTEXT("[/Script/EngineSettings]", "C47C7345459213852543F1906820C22C", "Simulation Blank"))
+GameDefaultMap=/Game/Main
+
+[/Script/Engine.RendererSettings]
+r.DefaultFeature.AutoExposure.ExtendDefaultLuminanceRange=True
+r.AllowStaticLighting=False
+r.Shadow.Virtual.Enable=1
+
+r.GenerateMeshDistanceFields=True
+
+r.DynamicGlobalIlluminationMethod=1
+
+r.ReflectionMethod=1
+
+[/Script/AndroidFileServerEditor.AndroidFileServerRuntimeSettings]
+bEnablePlugin=True
+bAllowNetworkConnection=True
+SecurityToken=7A2FAE694723660CF26AA1B1F56504C8
+bIncludeInShipping=False
+bAllowExternalStartInShipping=False
+bCompileAFSProject=False
+bUseCompression=False
+bLogFiles=False
+bReportStats=False
+ConnectionType=USBOnly
+bUseManualIPAddress=False
+ManualIPAddress=
+
+[/Script/WindowsTargetPlatform.WindowsTargetSettings]
+DefaultGraphicsRHI=DefaultGraphicsRHI_DX12
+-D3D12TargetedShaderFormats=PCD3D_SM5
++D3D12TargetedShaderFormats=PCD3D_SM6
+-D3D11TargetedShaderFormats=PCD3D_SM5
++D3D11TargetedShaderFormats=PCD3D_SM5
+Compiler=Default
+AudioSampleRate=48000
+AudioCallbackBufferFrameSize=1024
+AudioNumBuffersToEnqueue=1
+AudioMaxChannels=0
+AudioNumSourceWorkers=4
+SpatializationPlugin=
+SourceDataOverridePlugin=
+ReverbPlugin=
+OcclusionPlugin=
+CompressionOverrides=(bOverrideCompressionTimes=False,DurationThreshold=5.000000,MaxNumRandomBranches=0,SoundCueQualityIndex=0)
+CacheSizeKB=65536
+MaxChunkSizeOverrideKB=0
+bResampleForDevice=False
+MaxSampleRate=48000.000000
+HighSampleRate=32000.000000
+MedSampleRate=24000.000000
+LowSampleRate=12000.000000
+MinSampleRate=8000.000000
+CompressionQualityModifier=1.000000
+AutoStreamingThreshold=0.000000
+SoundCueCookQualityIndex=-1
diff --git a/Templates/TP_SIM_Sandbox/Config/DefaultGame.ini b/Templates/TP_SIM_Sandbox/Config/DefaultGame.ini
new file mode 100644
--- /dev/null
+++ b/Templates/TP_SIM_Sandbox/Config/DefaultGame.ini
@@ -0,0 +1,3 @@
+
+[/Script/EngineSettings.GeneralProjectSettings]
+ProjectName=SimSandbox
\ No newline at end of file
diff --git a/Templates/TP_SIM_Sandbox/Config/TemplateDefs.ini b/Templates/TP_SIM_Sandbox/Config/TemplateDefs.ini
new file mode 100644
--- /dev/null
+++ b/Templates/TP_SIM_Sandbox/Config/TemplateDefs.ini
@@ -0,0 +1,111 @@
+﻿; This file specifies how a template will be treated when creating a project from it.
+; It includes rules for ignoring files, replacing files and changing the contents of files.
+; This file is annotated to help with creating your own templates.
+
+[/Script/GameProjectGeneration.TemplateProjectDefs]
+
+; The categories that this template should appear in. The values in this list refer to the Key value given in TemplateCategories.ini.
+Categories=SIM
+
+; The localized display names for this template. Each entry specifies a Language code and Unicode text to display in the New Project dialog.
+LocalizedDisplayNames=(Language="en", Text="Sandbox")
+LocalizedDisplayNames=(Language="zh-Hans", Text="空白")
+LocalizedDisplayNames=(Language="ko", Text="기본")
+
+; The localized descriptions for this template, which is displayed when the template is selected.
+LocalizedDescriptions=(Language="en", Text="A sandbox environment with objects of different realistic shapes and sizes, made to define and test simulation models.\n\nBased on the Simulation Blank template.")
+LocalizedDescriptions=(Language="zh-Hans", Text="不含任何代码的空白项目。")
+LocalizedDescriptions=(Language="ja", Text="まったく空のプロジェクトです。コードは入っていません。")
+LocalizedDescriptions=(Language="ko", Text="코드가 포함되지 않은 기본 빈 프로젝트입니다.")
+
+; The sort key to use for the template to bypass the default alphabetical sorting.
+SortKey="_2"
+
+
+; This project uses a thumbnail as the project icon
+bThumbnailAsIcon=true
+
+; A list of folders to avoid copying when creating a project from this template. 
+; These should contain files only of concern to the template itself and not desired in the user's project.
+FoldersToIgnore=Binaries
+FoldersToIgnore=Build
+FoldersToIgnore=Intermediate
+FoldersToIgnore=Saved
+FoldersToIgnore=Media
+
+; Is this a blank template? Blank templates do not specify a default map.
+bIsBlank=False
+
+; The files to ignore copying, relative to the template root folder. 
+; %TEMPLATENAME% will be replaced by the template's folder name when resolving these. 
+; Filename comparisons are case-insensitive. 
+; For example, if this template is "TP_Blank" and the user creates "MyProject", then "TP_Blank.uproject" will be copied to "MyProject.uproject".
+FilesToIgnore="%TEMPLATENAME%.uproject"
+FilesToIgnore="%TEMPLATENAME%.png"
+FilesToIgnore="Config/TemplateDefs.ini"
+FilesToIgnore="Config/config.ini"
+FilesToIgnore="%TEMPLATENAME%.opensdf"
+FilesToIgnore="%TEMPLATENAME%.sdf"
+FilesToIgnore="%TEMPLATENAME%.v11.suo"
+FilesToIgnore="%TEMPLATENAME%.v12.suo"
+FilesToIgnore="%TEMPLATENAME%.sln"
+FilesToIgnore="Manifest.json"
+FilesToIgnore="contents.txt"
+
+
+; Folders that will be renamed when creating a project from this template.
+; %PROJECTNAME% will be replaced by the project name specified by the user during creation when resolving these rules.
+; For example,  then "Source/TP_Blank" will be copied to "Source/MyProject".
+FolderRenames=(From="Source/%TEMPLATENAME%", To="Source/%PROJECTNAME%")
+FolderRenames=(From="Source/%TEMPLATENAME%Editor", To="Source/%PROJECTNAME%Editor")
+
+; Rules for renaming files when creating a project from this template. 
+; %TEMPLATENAME_UPPERCASE% and %PROJECTNAME_UPPERCASE% are the template and project names in upper case, eg. for "TP_Blank", its upper case representation is "TP_BLANK".
+; %TEMPLATENAME_LOWERCASE% and %PROJECTNAME_LOWERCASE% are the template and project names in lower case, eg. for "TP_Blank", its lower case representation is "tp_blank".
+; A rule does not have to include any of these modifiers. 
+; Extensions is a list of files to consider for these replacements. Eg. ("cpp", "h") will try to apply a rule only to files like "SomeFile.cpp" and "SomeFile.h".
+; bCaseSensitive specifies whether this replacement should be case-sensitive.
+
+; Here are some default rules for renaming files:
+
+; This rule states that a file called "TP_BLANK.cpp" will be renamed to "MYPROJECT.cpp".
+FilenameReplacements=(Extensions=("cpp","h","ini","cs"), From="%TEMPLATENAME_UPPERCASE%", To="%PROJECTNAME_UPPERCASE%", bCaseSensitive=true)
+
+; This rule states that a file called "tp_blank.ini" will be renamed to "myproject.ini"
+FilenameReplacements=(Extensions=("cpp","h","ini","cs"), From="%TEMPLATENAME_LOWERCASE%", To="%PROJECTNAME_LOWERCASE%", bCaseSensitive=true)
+
+; And this rule catches any irregularly-cased files, so that "Tp_bLANK.ini" will be renamed to "MyProject.ini".
+FilenameReplacements=(Extensions=("cpp","h","ini","cs"), From="%TEMPLATENAME%",To="%PROJECTNAME%",bCaseSensitive=false)
+
+; Rules for renaming strings in the contents of files. 
+; The same rules as above apply, but any usages of strings that match these rules inside the file will be changed. This includes all substrings!
+; For example, "TP_Blank" will be replaced by "MyProject", but "TP_Blank_Name" will also be replaced by "MyProject_Name".
+ReplacementsInFiles=(Extensions=("cpp","h","ini","cs"), From="%TEMPLATENAME_UPPERCASE%", To="%PROJECTNAME_UPPERCASE%", bCaseSensitive=true)
+ReplacementsInFiles=(Extensions=("cpp","h","ini","cs"), From="%TEMPLATENAME_LOWERCASE%", To="%PROJECTNAME_LOWERCASE%", bCaseSensitive=true)
+ReplacementsInFiles=(Extensions=("cpp","h","ini","cs"), From="%TEMPLATENAME%", To="%PROJECTNAME%", bCaseSensitive=false)
+
+; The following settings are not enabled for this project, but are included for completeness:
+
+; Whether or not this template can be used to create a project. Default: true
+; If this is set to false, then the template becomes a Feature Pack that can only be referenced from other templates, or added to a project from the Content Browser.
+;bAllowProjectCreation=false
+
+; Which settings on the project settings page should be hidden for this template. Default: none
+; Allowed values are:
+; All - Hides the settings page entirely
+; Languages - Choose between C++ and Blueprint versions.
+; HardwareTarget - Choose between Mobile and Desktop hardware targets.
+; GraphicsPreset - Choose between Maximum and Scalable graphics.
+; StarterContent - Choose whether to add starter content to the project.
+; XR - Choose whether to enable Virtual Reality support for the project.
+;HiddenSettings=All
+;HiddenSettings=Languages
+HiddenSettings=HardwareTarget
+HiddenSettings=GraphicsPreset
+HiddenSettings=RayTracing
+;HiddenSettings=StarterContent 
+HiddenSettings=XR
+
+; Shared feature packs. The files in these packs listed in these structures marked as 'additionalfiles' will be copied on project generation.
+SharedContentPacks=(MountName="SimBlank",DetailLevels=("Standard"))
+SharedContentPacks=(MountName="SimSandbox",DetailLevels=("Standard"))
diff --git a/Templates/TP_SIM_Sandbox/Source/TP_SIM_Sandbox.Target.cs b/Templates/TP_SIM_Sandbox/Source/TP_SIM_Sandbox.Target.cs
new file mode 100644
--- /dev/null
+++ b/Templates/TP_SIM_Sandbox/Source/TP_SIM_Sandbox.Target.cs
@@ -0,0 +1,14 @@
+// Copyright Epic Games, Inc. All Rights Reserved.
+
+using UnrealBuildTool;
+using System.Collections.Generic;
+
+public class TP_SIM_SandboxTarget : TargetRules
+{
+	public TP_SIM_SandboxTarget( TargetInfo Target) : base(Target)
+	{
+		Type = TargetType.Game;
+		DefaultBuildSettings = BuildSettingsVersion.V2;
+		ExtraModuleNames.AddRange( new string[] { "TP_SIM_Sandbox" } );
+	}
+}
diff --git a/Templates/TP_SIM_Sandbox/Source/TP_SIM_Sandbox/TP_SIM_Sandbox.Build.cs b/Templates/TP_SIM_Sandbox/Source/TP_SIM_Sandbox/TP_SIM_Sandbox.Build.cs
new file mode 100644
--- /dev/null
+++ b/Templates/TP_SIM_Sandbox/Source/TP_SIM_Sandbox/TP_SIM_Sandbox.Build.cs
@@ -0,0 +1,23 @@
+// Copyright Epic Games, Inc. All Rights Reserved.
+
+using UnrealBuildTool;
+
+public class TP_SIM_Sandbox : ModuleRules
+{
+	public TP_SIM_Sandbox(ReadOnlyTargetRules Target) : base(Target)
+	{
+		PCHUsage = PCHUsageMode.UseExplicitOrSharedPCHs;
+	
+		PublicDependencyModuleNames.AddRange(new string[] { "Core", "CoreUObject", "Engine", "InputCore" });
+
+		PrivateDependencyModuleNames.AddRange(new string[] {  });
+
+		// Uncomment if you are using Slate UI
+		// PrivateDependencyModuleNames.AddRange(new string[] { "Slate", "SlateCore" });
+		
+		// Uncomment if you are using online features
+		// PrivateDependencyModuleNames.Add("OnlineSubsystem");
+
+		// To include OnlineSubsystemSteam, add it to the plugins section in your uproject file with the Enabled attribute set to true
+	}
+}
diff --git a/Templates/TP_SIM_Sandbox/Source/TP_SIM_Sandbox/TP_SIM_Sandbox.cpp b/Templates/TP_SIM_Sandbox/Source/TP_SIM_Sandbox/TP_SIM_Sandbox.cpp
new file mode 100644
--- /dev/null
+++ b/Templates/TP_SIM_Sandbox/Source/TP_SIM_Sandbox/TP_SIM_Sandbox.cpp
@@ -0,0 +1,6 @@
+// Copyright Epic Games, Inc. All Rights Reserved.
+
+#include "TP_SIM_Sandbox.h"
+#include "Modules/ModuleManager.h"
+
+IMPLEMENT_PRIMARY_GAME_MODULE( FDefaultGameModuleImpl, TP_SIM_Sandbox, "TP_SIM_Sandbox" );
diff --git a/Templates/TP_SIM_Sandbox/Source/TP_SIM_Sandbox/TP_SIM_Sandbox.h b/Templates/TP_SIM_Sandbox/Source/TP_SIM_Sandbox/TP_SIM_Sandbox.h
new file mode 100644
--- /dev/null
+++ b/Templates/TP_SIM_Sandbox/Source/TP_SIM_Sandbox/TP_SIM_Sandbox.h
@@ -0,0 +1,6 @@
+// Copyright Epic Games, Inc. All Rights Reserved.
+
+#pragma once
+
+#include "CoreMinimal.h"
+
diff --git a/Templates/TP_SIM_Sandbox/Source/TP_SIM_Sandbox/TP_SIM_SandboxGameModeBase.cpp b/Templates/TP_SIM_Sandbox/Source/TP_SIM_Sandbox/TP_SIM_SandboxGameModeBase.cpp
new file mode 100644
--- /dev/null
+++ b/Templates/TP_SIM_Sandbox/Source/TP_SIM_Sandbox/TP_SIM_SandboxGameModeBase.cpp
@@ -0,0 +1,5 @@
+// Copyright Epic Games, Inc. All Rights Reserved.
+
+
+#include "TP_SIM_SandboxGameModeBase.h"
+
diff --git a/Templates/TP_SIM_Sandbox/Source/TP_SIM_Sandbox/TP_SIM_SandboxGameModeBase.h b/Templates/TP_SIM_Sandbox/Source/TP_SIM_Sandbox/TP_SIM_SandboxGameModeBase.h
new file mode 100644
--- /dev/null
+++ b/Templates/TP_SIM_Sandbox/Source/TP_SIM_Sandbox/TP_SIM_SandboxGameModeBase.h
@@ -0,0 +1,17 @@
+// Copyright Epic Games, Inc. All Rights Reserved.
+
+#pragma once
+
+#include "CoreMinimal.h"
+#include "GameFramework/GameModeBase.h"
+#include "TP_SIM_SandboxGameModeBase.generated.h"
+
+/**
+ * 
+ */
+UCLASS()
+class TP_SIM_SANDBOX_API ATP_SIM_SandboxGameModeBase : public AGameModeBase
+{
+	GENERATED_BODY()
+	
+};
diff --git a/Templates/TP_SIM_Sandbox/Source/TP_SIM_SandboxEditor.Target.cs b/Templates/TP_SIM_Sandbox/Source/TP_SIM_SandboxEditor.Target.cs
new file mode 100644
--- /dev/null
+++ b/Templates/TP_SIM_Sandbox/Source/TP_SIM_SandboxEditor.Target.cs
@@ -0,0 +1,14 @@
+// Copyright Epic Games, Inc. All Rights Reserved.
+
+using UnrealBuildTool;
+using System.Collections.Generic;
+
+public class TP_SIM_SandboxEditorTarget : TargetRules
+{
+	public TP_SIM_SandboxEditorTarget( TargetInfo Target) : base(Target)
+	{
+		Type = TargetType.Editor;
+		DefaultBuildSettings = BuildSettingsVersion.V2;
+		ExtraModuleNames.AddRange( new string[] { "TP_SIM_Sandbox" } );
+	}
+}
diff --git a/Templates/TP_SIM_Sandbox/TP_SIM_Sandbox.uproject b/Templates/TP_SIM_Sandbox/TP_SIM_Sandbox.uproject
new file mode 100644
--- /dev/null
+++ b/Templates/TP_SIM_Sandbox/TP_SIM_Sandbox.uproject
@@ -0,0 +1,30 @@
+{
+	"FileVersion": 3,
+	"EngineAssociation": "",
+	"Category": "",
+	"Description": "",
+	"Modules": [
+		{
+			"Name": "TP_SIM_Sandbox",
+			"Type": "Runtime",
+			"LoadingPhase": "Default"
+		}
+	],
+	"Plugins": [
+		{
+			"Name": "ModelingToolsEditorMode",
+			"Enabled": true,
+			"TargetAllowList": [
+				"Editor"
+			]
+		},
+		{
+			"Name": "SunPosition",
+			"Enabled": true
+		},
+		{
+			"Name": "GeoReferencing",
+			"Enabled": true
+		}
+	]
+}
\ No newline at end of file
diff --git a/Templates/TP_SIM_Sandbox/manifest.json b/Templates/TP_SIM_Sandbox/manifest.json
new file mode 100644
--- /dev/null
+++ b/Templates/TP_SIM_Sandbox/manifest.json
@@ -0,0 +1,80 @@
+{
+	"Version": 1,
+	"Name":
+	[
+		{
+			"Language": "en",
+			"Text": "Simulation Sandbox"
+		},
+		{
+			"Language": "ko",
+			"Text": "일인칭"
+		},
+		{
+			"Language": "zh-Hans",
+			"Text": "第一人称游戏"
+		},
+		{
+			"Language": "ja",
+			"Text": "ファーストパーソン"
+		}
+	],
+	"Description":
+	[
+		{
+			"Language": "en",
+			"Text": "A sandbox environment with objects of different realistic shapes and sizes, made to define and test simulation models."
+		},
+		{
+			"Language": "ko",
+			"Text": "일인칭 템플릿은 일인칭 시점에서 보이는 (양팔로 표현되는) 플레이어 캐릭터를 특징으로 합니다. 캐릭터 이동은 키보드, 컨트롤러, 터치 디바이스의 가상 조이스틱으로 가능합니다. 추가적으로 플레이어는 마우스, 컨트롤러, 터치 디바이스의 가상 조이스틱을 사용하여 주변을 둘러볼 수도 있습니다. 캐릭터에는 총이 장착되어 있어 마우스, 컨트롤러, 터치 디바이스의 가상 조이스틱을 사용하면 간단한 발사체가 발사되어, 아레나 벽에 튕기면서 레벨에 있는 몇몇 물리 오브젝트에 영향을 줍니다."
+		},
+		{
+			"Language": "zh-Hans",
+			"Text": "第一人称模板包含一个第一人称玩家角色（通过一双手臂呈现）。你可以通过键盘、手柄或触摸式设备上的虚拟摇杆控制角色在关卡中四处移动。你还可以使用鼠标、手柄或触摸式设备上的虚拟摇杆转动角色视角，观察四周。此外，你还可以让角色拾取枪支，然后用鼠标、手柄或触摸式设备上的虚拟摇杆来发射球形子弹。子弹会撞飞关卡中的物理对象，还能从竞技场墙壁上弹开。"
+		},
+		{
+			"Language": "ja",
+			"Text": "ファーストパーソンテンプレートには、一人称視点で表示されるプレイヤーキャラクタ（腕として表示されています）がいます。キーボード、コントローラー、タッチデバイスでのバーチャルジョイスティックでこのキャラクタを移動させることができます。追加して、プレイヤーはマウス、コントローラー、タッチデバイスでのバーチャルジョイスティックで周りを見回すことができます。さらにキャラクタは銃を装備して、マウス、コントローラー、タッチデバイスでのバーチャルジョイスティックを使ってシンプルな球体の弾を発射できます。この弾は物理オブジェクトに影響し、アリーナの壁で反射します。"
+		}
+	],
+	"AssetTypes":
+	[
+		{
+			"Language": "en",
+			"Text": ""
+		}
+	],	
+	"SearchTags":
+	[
+		{
+			"Language": "en",
+			"Text": "Simulation"
+		}
+	],
+	"ClassTypes": "Actor",
+	"Category": "BlueprintFeature",
+	"FocusAsset": "",
+	"Thumbnail": "TP_SIM_Sandbox.png",
+	"Screenshots":
+	[
+		"TP_SIM_Sandbox_Preview.png"
+	],
+	"AdditionalFeaturePacks":
+	[
+		{
+			"MountName": "SimBlank",
+			"DetailLevels":
+			[
+				"Standard"
+			]
+		},
+		{
+			"MountName": "SimSandbox",
+			"DetailLevels":
+			[
+				"Standard"
+			]
+		}
+	]
+}
diff --git a/Templates/TP_SIM_SandboxBP/Config/DefaultEngine.ini b/Templates/TP_SIM_SandboxBP/Config/DefaultEngine.ini
new file mode 100644
--- /dev/null
+++ b/Templates/TP_SIM_SandboxBP/Config/DefaultEngine.ini
@@ -0,0 +1,61 @@
+
+
+[/Script/EngineSettings.GameMapsSettings]
+EditorStartupMap=/Game/SimSandbox/Levels/SimSandbox
++EditorTemplateMapOverrides=(Thumbnail="/Game/SimSandbox/Levels/Thumbnails/SimSandbox_Thumbnail.SimSandbox_Thumbnail",Map="/Game/SimSandbox/Levels/SimSandbox.SimSandbox",DisplayName=NSLOCTEXT("[/Script/EngineSettings]", "C47C7345459213852543F1906820C22C", "Simulation Sandbox"))
++EditorTemplateMapOverrides=(Thumbnail="/Game/SimBlank/Levels/Thumbnails/SimBlank_Thumbnail.SimBlank_Thumbnail",Map="/Game/SimBlank/Levels/SimBlank.SimBlank",DisplayName=NSLOCTEXT("[/Script/EngineSettings]", "C47C7345459213852543F1906820C22C", "Simulation Blank"))
+GameDefaultMap=/Game/Main
+
+[/Script/Engine.RendererSettings]
+r.DefaultFeature.AutoExposure.ExtendDefaultLuminanceRange=True
+r.AllowStaticLighting=False
+r.Shadow.Virtual.Enable=1
+
+r.GenerateMeshDistanceFields=True
+
+r.DynamicGlobalIlluminationMethod=1
+
+r.ReflectionMethod=1
+
+[/Script/AndroidFileServerEditor.AndroidFileServerRuntimeSettings]
+bEnablePlugin=True
+bAllowNetworkConnection=True
+SecurityToken=7A2FAE694723660CF26AA1B1F56504C8
+bIncludeInShipping=False
+bAllowExternalStartInShipping=False
+bCompileAFSProject=False
+bUseCompression=False
+bLogFiles=False
+bReportStats=False
+ConnectionType=USBOnly
+bUseManualIPAddress=False
+ManualIPAddress=
+
+[/Script/WindowsTargetPlatform.WindowsTargetSettings]
+DefaultGraphicsRHI=DefaultGraphicsRHI_DX12
+-D3D12TargetedShaderFormats=PCD3D_SM5
++D3D12TargetedShaderFormats=PCD3D_SM6
+-D3D11TargetedShaderFormats=PCD3D_SM5
++D3D11TargetedShaderFormats=PCD3D_SM5
+Compiler=Default
+AudioSampleRate=48000
+AudioCallbackBufferFrameSize=1024
+AudioNumBuffersToEnqueue=1
+AudioMaxChannels=0
+AudioNumSourceWorkers=4
+SpatializationPlugin=
+SourceDataOverridePlugin=
+ReverbPlugin=
+OcclusionPlugin=
+CompressionOverrides=(bOverrideCompressionTimes=False,DurationThreshold=5.000000,MaxNumRandomBranches=0,SoundCueQualityIndex=0)
+CacheSizeKB=65536
+MaxChunkSizeOverrideKB=0
+bResampleForDevice=False
+MaxSampleRate=48000.000000
+HighSampleRate=32000.000000
+MedSampleRate=24000.000000
+LowSampleRate=12000.000000
+MinSampleRate=8000.000000
+CompressionQualityModifier=1.000000
+AutoStreamingThreshold=0.000000
+SoundCueCookQualityIndex=-1
diff --git a/Templates/TP_SIM_SandboxBP/Config/DefaultGame.ini b/Templates/TP_SIM_SandboxBP/Config/DefaultGame.ini
new file mode 100644
--- /dev/null
+++ b/Templates/TP_SIM_SandboxBP/Config/DefaultGame.ini
@@ -0,0 +1,3 @@
+
+[/Script/EngineSettings.GeneralProjectSettings]
+ProjectName=SimSandboxBP
diff --git a/Templates/TP_SIM_SandboxBP/Config/TemplateDefs.ini b/Templates/TP_SIM_SandboxBP/Config/TemplateDefs.ini
new file mode 100644
--- /dev/null
+++ b/Templates/TP_SIM_SandboxBP/Config/TemplateDefs.ini
@@ -0,0 +1,111 @@
+﻿; This file specifies how a template will be treated when creating a project from it.
+; It includes rules for ignoring files, replacing files and changing the contents of files.
+; This file is annotated to help with creating your own templates.
+
+[/Script/GameProjectGeneration.TemplateProjectDefs]
+
+; The categories that this template should appear in. The values in this list refer to the Key value given in TemplateCategories.ini.
+Categories=SIM
+
+; The localized display names for this template. Each entry specifies a Language code and Unicode text to display in the New Project dialog.
+LocalizedDisplayNames=(Language="en", Text="Sandbox")
+LocalizedDisplayNames=(Language="zh-Hans", Text="空白")
+LocalizedDisplayNames=(Language="ko", Text="기본")
+
+; The localized descriptions for this template, which is displayed when the template is selected.
+LocalizedDescriptions=(Language="en", Text="A sandbox environment with objects of different realistic shapes and sizes, made to define and test simulation models.\n\nBased on the Simulation Blank template.")
+LocalizedDescriptions=(Language="zh-Hans", Text="不含任何代码的空白项目。")
+LocalizedDescriptions=(Language="ja", Text="まったく空のプロジェクトです。コードは入っていません。")
+LocalizedDescriptions=(Language="ko", Text="코드가 포함되지 않은 기본 빈 프로젝트입니다.")
+
+; The sort key to use for the template to bypass the default alphabetical sorting.
+SortKey="_2"
+
+
+; This project uses a thumbnail as the project icon
+bThumbnailAsIcon=true
+
+; A list of folders to avoid copying when creating a project from this template. 
+; These should contain files only of concern to the template itself and not desired in the user's project.
+FoldersToIgnore=Binaries
+FoldersToIgnore=Build
+FoldersToIgnore=Intermediate
+FoldersToIgnore=Saved
+FoldersToIgnore=Media
+
+; Is this a blank template? Blank templates do not specify a default map.
+bIsBlank=False
+
+; The files to ignore copying, relative to the template root folder. 
+; %TEMPLATENAME% will be replaced by the template's folder name when resolving these. 
+; Filename comparisons are case-insensitive. 
+; For example, if this template is "TP_Blank" and the user creates "MyProject", then "TP_Blank.uproject" will be copied to "MyProject.uproject".
+FilesToIgnore="%TEMPLATENAME%.uproject"
+FilesToIgnore="%TEMPLATENAME%.png"
+FilesToIgnore="Config/TemplateDefs.ini"
+FilesToIgnore="Config/config.ini"
+FilesToIgnore="%TEMPLATENAME%.opensdf"
+FilesToIgnore="%TEMPLATENAME%.sdf"
+FilesToIgnore="%TEMPLATENAME%.v11.suo"
+FilesToIgnore="%TEMPLATENAME%.v12.suo"
+FilesToIgnore="%TEMPLATENAME%.sln"
+FilesToIgnore="Manifest.json"
+FilesToIgnore="contents.txt"
+
+
+; Folders that will be renamed when creating a project from this template.
+; %PROJECTNAME% will be replaced by the project name specified by the user during creation when resolving these rules.
+; For example,  then "Source/TP_Blank" will be copied to "Source/MyProject".
+FolderRenames=(From="Source/%TEMPLATENAME%", To="Source/%PROJECTNAME%")
+FolderRenames=(From="Source/%TEMPLATENAME%Editor", To="Source/%PROJECTNAME%Editor")
+
+; Rules for renaming files when creating a project from this template. 
+; %TEMPLATENAME_UPPERCASE% and %PROJECTNAME_UPPERCASE% are the template and project names in upper case, eg. for "TP_Blank", its upper case representation is "TP_BLANK".
+; %TEMPLATENAME_LOWERCASE% and %PROJECTNAME_LOWERCASE% are the template and project names in lower case, eg. for "TP_Blank", its lower case representation is "tp_blank".
+; A rule does not have to include any of these modifiers. 
+; Extensions is a list of files to consider for these replacements. Eg. ("cpp", "h") will try to apply a rule only to files like "SomeFile.cpp" and "SomeFile.h".
+; bCaseSensitive specifies whether this replacement should be case-sensitive.
+
+; Here are some default rules for renaming files:
+
+; This rule states that a file called "TP_BLANK.cpp" will be renamed to "MYPROJECT.cpp".
+FilenameReplacements=(Extensions=("cpp","h","ini","cs"), From="%TEMPLATENAME_UPPERCASE%", To="%PROJECTNAME_UPPERCASE%", bCaseSensitive=true)
+
+; This rule states that a file called "tp_blank.ini" will be renamed to "myproject.ini"
+FilenameReplacements=(Extensions=("cpp","h","ini","cs"), From="%TEMPLATENAME_LOWERCASE%", To="%PROJECTNAME_LOWERCASE%", bCaseSensitive=true)
+
+; And this rule catches any irregularly-cased files, so that "Tp_bLANK.ini" will be renamed to "MyProject.ini".
+FilenameReplacements=(Extensions=("cpp","h","ini","cs"), From="%TEMPLATENAME%",To="%PROJECTNAME%",bCaseSensitive=false)
+
+; Rules for renaming strings in the contents of files. 
+; The same rules as above apply, but any usages of strings that match these rules inside the file will be changed. This includes all substrings!
+; For example, "TP_Blank" will be replaced by "MyProject", but "TP_Blank_Name" will also be replaced by "MyProject_Name".
+ReplacementsInFiles=(Extensions=("cpp","h","ini","cs"), From="%TEMPLATENAME_UPPERCASE%", To="%PROJECTNAME_UPPERCASE%", bCaseSensitive=true)
+ReplacementsInFiles=(Extensions=("cpp","h","ini","cs"), From="%TEMPLATENAME_LOWERCASE%", To="%PROJECTNAME_LOWERCASE%", bCaseSensitive=true)
+ReplacementsInFiles=(Extensions=("cpp","h","ini","cs"), From="%TEMPLATENAME%", To="%PROJECTNAME%", bCaseSensitive=false)
+
+; The following settings are not enabled for this project, but are included for completeness:
+
+; Whether or not this template can be used to create a project. Default: true
+; If this is set to false, then the template becomes a Feature Pack that can only be referenced from other templates, or added to a project from the Content Browser.
+;bAllowProjectCreation=false
+
+; Which settings on the project settings page should be hidden for this template. Default: none
+; Allowed values are:
+; All - Hides the settings page entirely
+; Languages - Choose between C++ and Blueprint versions.
+; HardwareTarget - Choose between Mobile and Desktop hardware targets.
+; GraphicsPreset - Choose between Maximum and Scalable graphics.
+; StarterContent - Choose whether to add starter content to the project.
+; XR - Choose whether to enable Virtual Reality support for the project.
+;HiddenSettings=All
+;HiddenSettings=Languages
+HiddenSettings=HardwareTarget
+HiddenSettings=GraphicsPreset
+HiddenSettings=RayTracing
+;HiddenSettings=StarterContent 
+HiddenSettings=XR
+
+; Shared feature packs. The files in these packs listed in these structures marked as 'additionalfiles' will be copied on project generation.
+SharedContentPacks=(MountName="SimBlank",DetailLevels=("Standard"))
+SharedContentPacks=(MountName="SimSandbox",DetailLevels=("Standard"))
diff --git a/Templates/TP_SIM_SandboxBP/TP_SIM_SandboxBP.uproject b/Templates/TP_SIM_SandboxBP/TP_SIM_SandboxBP.uproject
new file mode 100644
--- /dev/null
+++ b/Templates/TP_SIM_SandboxBP/TP_SIM_SandboxBP.uproject
@@ -0,0 +1,23 @@
+{
+	"FileVersion": 3,
+	"EngineAssociation": "",
+	"Category": "",
+	"Description": "",
+	"Plugins": [
+		{
+			"Name": "ModelingToolsEditorMode",
+			"Enabled": true,
+			"TargetAllowList": [
+				"Editor"
+			]
+		},
+		{
+			"Name": "SunPosition",
+			"Enabled": true
+		},
+		{
+			"Name": "GeoReferencing",
+			"Enabled": true
+		}
+	]
+}
\ No newline at end of file
diff --git a/Templates/TP_SIM_SandboxBP/manifest.json b/Templates/TP_SIM_SandboxBP/manifest.json
new file mode 100644
--- /dev/null
+++ b/Templates/TP_SIM_SandboxBP/manifest.json
@@ -0,0 +1,80 @@
+{
+	"Version": 1,
+	"Name":
+	[
+		{
+			"Language": "en",
+			"Text": "Simulation Sandbox"
+		},
+		{
+			"Language": "ko",
+			"Text": "일인칭"
+		},
+		{
+			"Language": "zh-Hans",
+			"Text": "第一人称游戏"
+		},
+		{
+			"Language": "ja",
+			"Text": "ファーストパーソン"
+		}
+	],
+	"Description":
+	[
+		{
+			"Language": "en",
+			"Text": "A sandbox environment with objects of different realistic shapes and sizes, made to define and test simulation models."
+		},
+		{
+			"Language": "ko",
+			"Text": "일인칭 템플릿은 일인칭 시점에서 보이는 (양팔로 표현되는) 플레이어 캐릭터를 특징으로 합니다. 캐릭터 이동은 키보드, 컨트롤러, 터치 디바이스의 가상 조이스틱으로 가능합니다. 추가적으로 플레이어는 마우스, 컨트롤러, 터치 디바이스의 가상 조이스틱을 사용하여 주변을 둘러볼 수도 있습니다. 캐릭터에는 총이 장착되어 있어 마우스, 컨트롤러, 터치 디바이스의 가상 조이스틱을 사용하면 간단한 발사체가 발사되어, 아레나 벽에 튕기면서 레벨에 있는 몇몇 물리 오브젝트에 영향을 줍니다."
+		},
+		{
+			"Language": "zh-Hans",
+			"Text": "第一人称模板包含一个第一人称玩家角色（通过一双手臂呈现）。你可以通过键盘、手柄或触摸式设备上的虚拟摇杆控制角色在关卡中四处移动。你还可以使用鼠标、手柄或触摸式设备上的虚拟摇杆转动角色视角，观察四周。此外，你还可以让角色拾取枪支，然后用鼠标、手柄或触摸式设备上的虚拟摇杆来发射球形子弹。子弹会撞飞关卡中的物理对象，还能从竞技场墙壁上弹开。"
+		},
+		{
+			"Language": "ja",
+			"Text": "ファーストパーソンテンプレートには、一人称視点で表示されるプレイヤーキャラクタ（腕として表示されています）がいます。キーボード、コントローラー、タッチデバイスでのバーチャルジョイスティックでこのキャラクタを移動させることができます。追加して、プレイヤーはマウス、コントローラー、タッチデバイスでのバーチャルジョイスティックで周りを見回すことができます。さらにキャラクタは銃を装備して、マウス、コントローラー、タッチデバイスでのバーチャルジョイスティックを使ってシンプルな球体の弾を発射できます。この弾は物理オブジェクトに影響し、アリーナの壁で反射します。"
+		}
+	],
+	"AssetTypes":
+	[
+		{
+			"Language": "en",
+			"Text": ""
+		}
+	],	
+	"SearchTags":
+	[
+		{
+			"Language": "en",
+			"Text": "Simulation"
+		}
+	],
+	"ClassTypes": "Actor",
+	"Category": "BlueprintFeature",
+	"FocusAsset": "",
+	"Thumbnail": "TP_SIM_SandboxBP.png",
+	"Screenshots":
+	[
+		"TP_SIM_SandboxBP_Preview.png"
+	],
+	"AdditionalFeaturePacks":
+	[
+		{
+			"MountName": "SimBlank",
+			"DetailLevels":
+			[
+				"Standard"
+			]
+		},
+		{
+			"MountName": "SimSandbox",
+			"DetailLevels":
+			[
+				"Standard"
+			]
+		}
+	]
+}
diff --git a/Templates/TP_VirtualRealityBP/Config/TemplateDefs.ini b/Templates/TP_VirtualRealityBP/Config/TemplateDefs.ini
--- a/Templates/TP_VirtualRealityBP/Config/TemplateDefs.ini
+++ b/Templates/TP_VirtualRealityBP/Config/TemplateDefs.ini
@@ -1,6 +1,7 @@
 ﻿[/Script/GameProjectGeneration.TemplateProjectDefs]
 
 Categories=Games
+Categories=Sim
 
 LocalizedDisplayNames=(Language="en",Text="Virtual Reality")
 LocalizedDescriptions=(Language="en",Text="Blueprint Virtual Reality Template using OpenXR for Desktop, Console, and Mobile. The template features teleport locomotion, snap turning, grabbable objects, grabbable and interactive objects, Google Resonance spatialized audio, and a VR Spectator Camera.")
diff --git a/Templates/TemplateResources/Standard/SimBlank/FeaturePack/manifest.json b/Templates/TemplateResources/Standard/SimBlank/FeaturePack/manifest.json
new file mode 100644
--- /dev/null
+++ b/Templates/TemplateResources/Standard/SimBlank/FeaturePack/manifest.json
@@ -0,0 +1,48 @@
+{
+	"Version": 1,
+	"Name":
+	[
+		{
+			"Language": "en",
+			"Text": "Simulation Blank template"
+		}
+	],
+	"Description":
+	[
+		{
+			"Language": "en",
+			"Text": "Shared Resource Pack for Simulation Blank template"
+		}
+	],
+	"AssetTypes":
+	[
+		{
+			"Language": "en",
+			"Text": ""
+		}
+	],
+	"SearchTags":
+	[
+		{
+			"Language": "en",
+			"Text": ""
+		}
+	],	
+	"ClassTypes": " ",
+	"Category": "SharedPack",
+	"Thumbnail": "SimBlank.png",
+	"Screenshots":
+	[
+		"SimBlank_Preview.png"
+	],
+	"Ident": "SimBlankStandard",
+	"AdditionalFiles":
+	{
+		"DestinationFilesFolder": "SimBlank",
+		"AdditionalFilesList":
+		[
+			"Templates/TemplateResources/Standard/SimBlank/Content/*.*"
+		]
+	}
+}
+
diff --git a/Templates/TemplateResources/Standard/SimDefault/FeaturePack/manifest.json b/Templates/TemplateResources/Standard/SimDefault/FeaturePack/manifest.json
deleted file mode 100644
diff --git a/Templates/TemplateResources/Standard/SimSandbox/FeaturePack/manifest.json b/Templates/TemplateResources/Standard/SimSandbox/FeaturePack/manifest.json
new file mode 100644
--- /dev/null
+++ b/Templates/TemplateResources/Standard/SimSandbox/FeaturePack/manifest.json
@@ -0,0 +1,48 @@
+{
+	"Version": 1,
+	"Name":
+	[
+		{
+			"Language": "en",
+			"Text": "Simulation templates Shared Resource - Sandbox template"
+		}
+	],
+	"Description":
+	[
+		{
+			"Language": "en",
+			"Text": "Shared Resource Pack for Simulation Sandbox template"
+		}
+	],
+	"AssetTypes":
+	[
+		{
+			"Language": "en",
+			"Text": ""
+		}
+	],
+	"SearchTags":
+	[
+		{
+			"Language": "en",
+			"Text": ""
+		}
+	],	
+	"ClassTypes": " ",
+	"Category": "SharedPack",
+	"Thumbnail": "SimSandbox.png",
+	"Screenshots":
+	[
+		"SimSandbox_Preview.png"
+	],
+	"Ident": "SimSandboxStandard",
+	"AdditionalFiles":
+	{
+		"DestinationFilesFolder": "SimSandbox",
+		"AdditionalFilesList":
+		[
+			"Templates/TemplateResources/Standard/SimSandbox/Content/*.*"
+		]
+	}
+}
+
