Forcing DX12/SM6 to enable VSM in Simulation templates
Content Only

diff --git a/Templates/TP_SIM_Blank/Config/DefaultEngine.ini b/Templates/TP_SIM_Blank/Config/DefaultEngine.ini
--- a/Templates/TP_SIM_Blank/Config/DefaultEngine.ini
+++ b/Templates/TP_SIM_Blank/Config/DefaultEngine.ini
@@ -32,10 +32,8 @@ ManualIPAddress=
 
 [/Script/WindowsTargetPlatform.WindowsTargetSettings]
 DefaultGraphicsRHI=DefaultGraphicsRHI_DX12
--D3D12TargetedShaderFormats=PCD3D_SM5
-+D3D12TargetedShaderFormats=PCD3D_SM6
--D3D11TargetedShaderFormats=PCD3D_SM5
-+D3D11TargetedShaderFormats=PCD3D_SM5
+D3D12TargetedShaderFormats=PCD3D_SM6
+D3D11TargetedShaderFormats=PCD3D_SM5
 Compiler=Default
 AudioSampleRate=48000
 AudioCallbackBufferFrameSize=1024
diff --git a/Templates/TP_SIM_BlankBP/Config/DefaultEngine.ini b/Templates/TP_SIM_BlankBP/Config/DefaultEngine.ini
--- a/Templates/TP_SIM_BlankBP/Config/DefaultEngine.ini
+++ b/Templates/TP_SIM_BlankBP/Config/DefaultEngine.ini
@@ -34,10 +34,8 @@ ManualIPAddress=
 
 [/Script/WindowsTargetPlatform.WindowsTargetSettings]
 DefaultGraphicsRHI=DefaultGraphicsRHI_DX12
--D3D12TargetedShaderFormats=PCD3D_SM5
-+D3D12TargetedShaderFormats=PCD3D_SM6
--D3D11TargetedShaderFormats=PCD3D_SM5
-+D3D11TargetedShaderFormats=PCD3D_SM5
+D3D12TargetedShaderFormats=PCD3D_SM6
+D3D11TargetedShaderFormats=PCD3D_SM5
 Compiler=Default
 AudioSampleRate=48000
 AudioCallbackBufferFrameSize=1024
diff --git a/Templates/TP_SIM_Sandbox/Config/DefaultEngine.ini b/Templates/TP_SIM_Sandbox/Config/DefaultEngine.ini
--- a/Templates/TP_SIM_Sandbox/Config/DefaultEngine.ini
+++ b/Templates/TP_SIM_Sandbox/Config/DefaultEngine.ini
@@ -34,10 +34,8 @@ ManualIPAddress=
 
 [/Script/WindowsTargetPlatform.WindowsTargetSettings]
 DefaultGraphicsRHI=DefaultGraphicsRHI_DX12
--D3D12TargetedShaderFormats=PCD3D_SM5
-+D3D12TargetedShaderFormats=PCD3D_SM6
--D3D11TargetedShaderFormats=PCD3D_SM5
-+D3D11TargetedShaderFormats=PCD3D_SM5
+D3D12TargetedShaderFormats=PCD3D_SM6
+D3D11TargetedShaderFormats=PCD3D_SM5
 Compiler=Default
 AudioSampleRate=48000
 AudioCallbackBufferFrameSize=1024
diff --git a/Templates/TP_SIM_SandboxBP/Config/DefaultEngine.ini b/Templates/TP_SIM_SandboxBP/Config/DefaultEngine.ini
--- a/Templates/TP_SIM_SandboxBP/Config/DefaultEngine.ini
+++ b/Templates/TP_SIM_SandboxBP/Config/DefaultEngine.ini
@@ -34,10 +34,8 @@ ManualIPAddress=
 
 [/Script/WindowsTargetPlatform.WindowsTargetSettings]
 DefaultGraphicsRHI=DefaultGraphicsRHI_DX12
--D3D12TargetedShaderFormats=PCD3D_SM5
-+D3D12TargetedShaderFormats=PCD3D_SM6
--D3D11TargetedShaderFormats=PCD3D_SM5
-+D3D11TargetedShaderFormats=PCD3D_SM5
+D3D12TargetedShaderFormats=PCD3D_SM6
+D3D11TargetedShaderFormats=PCD3D_SM5
 Compiler=Default
 AudioSampleRate=48000
 AudioCallbackBufferFrameSize=1024
