Remove AndroidFileServer sections

#ue5

diff --git a/Templates/TP_AEC_ArchvisBP/Config/DefaultEngine.ini b/Templates/TP_AEC_ArchvisBP/Config/DefaultEngine.ini
--- a/Templates/TP_AEC_ArchvisBP/Config/DefaultEngine.ini
+++ b/Templates/TP_AEC_ArchvisBP/Config/DefaultEngine.ini
@@ -202,18 +202,3 @@ MinSampleRate=8000.000000
 CompressionQualityModifier=1.000000
 AutoStreamingThreshold=0.000000
 SoundCueCookQualityIndex=-1
-
-[/Script/AndroidFileServerEditor.AndroidFileServerRuntimeSettings]
-bEnablePlugin=True
-bAllowNetworkConnection=True
-SecurityToken=1FB1AEC34D02F7CC07305E90E62356F9
-bIncludeInShipping=False
-bAllowExternalStartInShipping=False
-bCompileAFSProject=False
-bUseCompression=False
-bLogFiles=False
-bReportStats=False
-ConnectionType=USBOnly
-bUseManualIPAddress=False
-ManualIPAddress=
-
diff --git a/Templates/TP_AEC_BlankBP/Config/DefaultEngine.ini b/Templates/TP_AEC_BlankBP/Config/DefaultEngine.ini
--- a/Templates/TP_AEC_BlankBP/Config/DefaultEngine.ini
+++ b/Templates/TP_AEC_BlankBP/Config/DefaultEngine.ini
@@ -151,18 +151,3 @@ SoundCueCookQualityIndex=-1
 +CollisionChannelRedirects=(OldName="Dynamic",NewName="WorldDynamic")
 +CollisionChannelRedirects=(OldName="VehicleMovement",NewName="Vehicle")
 +CollisionChannelRedirects=(OldName="PawnMovement",NewName="Pawn")
-
-[/Script/AndroidFileServerEditor.AndroidFileServerRuntimeSettings]
-bEnablePlugin=True
-bAllowNetworkConnection=True
-SecurityToken=C50CAFA9400E6E26E46B36B62632BFBD
-bIncludeInShipping=False
-bAllowExternalStartInShipping=False
-bCompileAFSProject=False
-bUseCompression=False
-bLogFiles=False
-bReportStats=False
-ConnectionType=USBOnly
-bUseManualIPAddress=False
-ManualIPAddress=
-
diff --git a/Templates/TP_Blank/Config/DefaultEngine.ini b/Templates/TP_Blank/Config/DefaultEngine.ini
--- a/Templates/TP_Blank/Config/DefaultEngine.ini
+++ b/Templates/TP_Blank/Config/DefaultEngine.ini
@@ -31,18 +31,3 @@ MinSampleRate=8000.000000
 CompressionQualityModifier=1.000000
 AutoStreamingThreshold=0.000000
 SoundCueCookQualityIndex=-1
-
-[/Script/AndroidFileServerEditor.AndroidFileServerRuntimeSettings]
-bEnablePlugin=True
-bAllowNetworkConnection=True
-SecurityToken=DED2580048DBB2163D6FA29747ACD6DB
-bIncludeInShipping=False
-bAllowExternalStartInShipping=False
-bCompileAFSProject=False
-bUseCompression=False
-bLogFiles=False
-bReportStats=False
-ConnectionType=USBOnly
-bUseManualIPAddress=False
-ManualIPAddress=
-
diff --git a/Templates/TP_FirstPersonBP/Config/DefaultEngine.ini b/Templates/TP_FirstPersonBP/Config/DefaultEngine.ini
--- a/Templates/TP_FirstPersonBP/Config/DefaultEngine.ini
+++ b/Templates/TP_FirstPersonBP/Config/DefaultEngine.ini
@@ -22,18 +22,3 @@ r.GenerateMeshDistanceFields=True
 r.DynamicGlobalIlluminationMethod=1
 r.Mobile.EnableNoPrecomputedLightingCSMShader=1
 r.DefaultFeature.AutoExposure.ExtendDefaultLuminanceRange=true
-
-[/Script/AndroidFileServerEditor.AndroidFileServerRuntimeSettings]
-bEnablePlugin=True
-bAllowNetworkConnection=True
-SecurityToken=5011EBFD4DF26CEF23A820976F45E941
-bIncludeInShipping=False
-bAllowExternalStartInShipping=False
-bCompileAFSProject=False
-bUseCompression=False
-bLogFiles=False
-bReportStats=False
-ConnectionType=USBOnly
-bUseManualIPAddress=False
-ManualIPAddress=
-
diff --git a/Templates/TP_SIM_Blank/Config/DefaultEngine.ini b/Templates/TP_SIM_Blank/Config/DefaultEngine.ini
--- a/Templates/TP_SIM_Blank/Config/DefaultEngine.ini
+++ b/Templates/TP_SIM_Blank/Config/DefaultEngine.ini
@@ -9,27 +9,10 @@ GameDefaultMap=/Game/Main
 r.DefaultFeature.AutoExposure.ExtendDefaultLuminanceRange=True
 r.AllowStaticLighting=False
 r.Shadow.Virtual.Enable=1
-
 r.GenerateMeshDistanceFields=True
-
 r.DynamicGlobalIlluminationMethod=1
-
 r.ReflectionMethod=1
 
-[/Script/AndroidFileServerEditor.AndroidFileServerRuntimeSettings]
-bEnablePlugin=True
-bAllowNetworkConnection=True
-SecurityToken=7A2FAE694723660CF26AA1B1F56504C8
-bIncludeInShipping=False
-bAllowExternalStartInShipping=False
-bCompileAFSProject=False
-bUseCompression=False
-bLogFiles=False
-bReportStats=False
-ConnectionType=USBOnly
-bUseManualIPAddress=False
-ManualIPAddress=
-
 [/Script/WindowsTargetPlatform.WindowsTargetSettings]
 DefaultGraphicsRHI=DefaultGraphicsRHI_DX12
 D3D12TargetedShaderFormats=PCD3D_SM6
diff --git a/Templates/TP_SIM_BlankBP/Config/DefaultEngine.ini b/Templates/TP_SIM_BlankBP/Config/DefaultEngine.ini
--- a/Templates/TP_SIM_BlankBP/Config/DefaultEngine.ini
+++ b/Templates/TP_SIM_BlankBP/Config/DefaultEngine.ini
@@ -11,27 +11,10 @@ GameDefaultMap=/Game/SimBlank/Levels/SimBlank.SimBlank
 r.DefaultFeature.AutoExposure.ExtendDefaultLuminanceRange=True
 r.AllowStaticLighting=False
 r.Shadow.Virtual.Enable=1
-
 r.GenerateMeshDistanceFields=True
-
 r.DynamicGlobalIlluminationMethod=1
-
 r.ReflectionMethod=1
 
-[/Script/AndroidFileServerEditor.AndroidFileServerRuntimeSettings]
-bEnablePlugin=True
-bAllowNetworkConnection=True
-SecurityToken=7A2FAE694723660CF26AA1B1F56504C8
-bIncludeInShipping=False
-bAllowExternalStartInShipping=False
-bCompileAFSProject=False
-bUseCompression=False
-bLogFiles=False
-bReportStats=False
-ConnectionType=USBOnly
-bUseManualIPAddress=False
-ManualIPAddress=
-
 [/Script/WindowsTargetPlatform.WindowsTargetSettings]
 DefaultGraphicsRHI=DefaultGraphicsRHI_DX12
 D3D12TargetedShaderFormats=PCD3D_SM6
diff --git a/Templates/TP_SIM_Sandbox/Config/DefaultEngine.ini b/Templates/TP_SIM_Sandbox/Config/DefaultEngine.ini
--- a/Templates/TP_SIM_Sandbox/Config/DefaultEngine.ini
+++ b/Templates/TP_SIM_Sandbox/Config/DefaultEngine.ini
@@ -11,27 +11,10 @@ GameDefaultMap=/Game/SimSandbox/Levels/SimSandbox.SimSandbox
 r.DefaultFeature.AutoExposure.ExtendDefaultLuminanceRange=True
 r.AllowStaticLighting=False
 r.Shadow.Virtual.Enable=1
-
 r.GenerateMeshDistanceFields=True
-
 r.DynamicGlobalIlluminationMethod=1
-
 r.ReflectionMethod=1
 
-[/Script/AndroidFileServerEditor.AndroidFileServerRuntimeSettings]
-bEnablePlugin=True
-bAllowNetworkConnection=True
-SecurityToken=7A2FAE694723660CF26AA1B1F56504C8
-bIncludeInShipping=False
-bAllowExternalStartInShipping=False
-bCompileAFSProject=False
-bUseCompression=False
-bLogFiles=False
-bReportStats=False
-ConnectionType=USBOnly
-bUseManualIPAddress=False
-ManualIPAddress=
-
 [/Script/WindowsTargetPlatform.WindowsTargetSettings]
 DefaultGraphicsRHI=DefaultGraphicsRHI_DX12
 D3D12TargetedShaderFormats=PCD3D_SM6
diff --git a/Templates/TP_SIM_SandboxBP/Config/DefaultEngine.ini b/Templates/TP_SIM_SandboxBP/Config/DefaultEngine.ini
--- a/Templates/TP_SIM_SandboxBP/Config/DefaultEngine.ini
+++ b/Templates/TP_SIM_SandboxBP/Config/DefaultEngine.ini
@@ -11,27 +11,10 @@ GameDefaultMap=/Game/SimSandbox/Levels/SimSandbox.SimSandbox
 r.DefaultFeature.AutoExposure.ExtendDefaultLuminanceRange=True
 r.AllowStaticLighting=False
 r.Shadow.Virtual.Enable=1
-
 r.GenerateMeshDistanceFields=True
-
 r.DynamicGlobalIlluminationMethod=1
-
 r.ReflectionMethod=1
 
-[/Script/AndroidFileServerEditor.AndroidFileServerRuntimeSettings]
-bEnablePlugin=True
-bAllowNetworkConnection=True
-SecurityToken=7A2FAE694723660CF26AA1B1F56504C8
-bIncludeInShipping=False
-bAllowExternalStartInShipping=False
-bCompileAFSProject=False
-bUseCompression=False
-bLogFiles=False
-bReportStats=False
-ConnectionType=USBOnly
-bUseManualIPAddress=False
-ManualIPAddress=
-
 [/Script/WindowsTargetPlatform.WindowsTargetSettings]
 DefaultGraphicsRHI=DefaultGraphicsRHI_DX12
 D3D12TargetedShaderFormats=PCD3D_SM6
diff --git a/Templates/TP_ThirdPerson/Config/DefaultEngine.ini b/Templates/TP_ThirdPerson/Config/DefaultEngine.ini
--- a/Templates/TP_ThirdPerson/Config/DefaultEngine.ini
+++ b/Templates/TP_ThirdPerson/Config/DefaultEngine.ini
@@ -40,18 +40,3 @@ MinSampleRate=8000.000000
 CompressionQualityModifier=1.000000
 AutoStreamingThreshold=0.000000
 SoundCueCookQualityIndex=-1
-
-[/Script/AndroidFileServerEditor.AndroidFileServerRuntimeSettings]
-bEnablePlugin=True
-bAllowNetworkConnection=True
-SecurityToken=BF6D5A5549DB55B6261E86A891EB414B
-bIncludeInShipping=False
-bAllowExternalStartInShipping=False
-bCompileAFSProject=False
-bUseCompression=False
-bLogFiles=False
-bReportStats=False
-ConnectionType=USBOnly
-bUseManualIPAddress=False
-ManualIPAddress=
-
diff --git a/Templates/TP_ThirdPersonBP/Config/DefaultEngine.ini b/Templates/TP_ThirdPersonBP/Config/DefaultEngine.ini
--- a/Templates/TP_ThirdPersonBP/Config/DefaultEngine.ini
+++ b/Templates/TP_ThirdPersonBP/Config/DefaultEngine.ini
@@ -48,18 +48,3 @@ MinSampleRate=8000.000000
 CompressionQualityModifier=1.000000
 AutoStreamingThreshold=0.000000
 SoundCueCookQualityIndex=-1
-
-[/Script/AndroidFileServerEditor.AndroidFileServerRuntimeSettings]
-bEnablePlugin=True
-bAllowNetworkConnection=True
-SecurityToken=22840D9C40E8F8CFF3280483D90DEE70
-bIncludeInShipping=False
-bAllowExternalStartInShipping=False
-bCompileAFSProject=False
-bUseCompression=False
-bLogFiles=False
-bReportStats=False
-ConnectionType=USBOnly
-bUseManualIPAddress=False
-ManualIPAddress=
-
diff --git a/Templates/TP_TopDown/Config/DefaultEngine.ini b/Templates/TP_TopDown/Config/DefaultEngine.ini
--- a/Templates/TP_TopDown/Config/DefaultEngine.ini
+++ b/Templates/TP_TopDown/Config/DefaultEngine.ini
@@ -7,20 +7,6 @@ GlobalDefaultGameMode="/Script/TP_TopDown.TP_TopDownGameMode"
 r.Mobile.EnableNoPrecomputedLightingCSMShader=1
  r.DefaultFeature.AutoExposure.ExtendDefaultLuminanceRange=true
 
-[/Script/AndroidFileServerEditor.AndroidFileServerRuntimeSettings]
-bEnablePlugin=True
-bAllowNetworkConnection=True
-SecurityToken=60F4C6BA45270F82910E719121C45142
-bIncludeInShipping=False
-bAllowExternalStartInShipping=False
-bCompileAFSProject=False
-bUseCompression=False
-bLogFiles=False
-bReportStats=False
-ConnectionType=USBOnly
-bUseManualIPAddress=False
-ManualIPAddress=
-
 [/Script/WindowsTargetPlatform.WindowsTargetSettings]
 DefaultGraphicsRHI=DefaultGraphicsRHI_DX12
 -D3D12TargetedShaderFormats=PCD3D_SM5
diff --git a/Templates/TP_TopDownBP/Config/DefaultEngine.ini b/Templates/TP_TopDownBP/Config/DefaultEngine.ini
--- a/Templates/TP_TopDownBP/Config/DefaultEngine.ini
+++ b/Templates/TP_TopDownBP/Config/DefaultEngine.ini
@@ -57,20 +57,6 @@ r.Shadow.Virtual.Enable=1
 r.Mobile.EnableNoPrecomputedLightingCSMShader=1
 r.DefaultFeature.AutoExposure.ExtendDefaultLuminanceRange=true
 
-[/Script/AndroidFileServerEditor.AndroidFileServerRuntimeSettings]
-bEnablePlugin=True
-bAllowNetworkConnection=True
-SecurityToken=CAF39D9941794497898047B0405D7FA8
-bIncludeInShipping=False
-bAllowExternalStartInShipping=False
-bCompileAFSProject=False
-bUseCompression=False
-bLogFiles=False
-bReportStats=False
-ConnectionType=USBOnly
-bUseManualIPAddress=False
-ManualIPAddress=
-
 [/Script/WindowsTargetPlatform.WindowsTargetSettings]
 DefaultGraphicsRHI=DefaultGraphicsRHI_DX12
 -D3D12TargetedShaderFormats=PCD3D_SM5
diff --git a/Templates/TP_VehicleAdv/Config/DefaultEngine.ini b/Templates/TP_VehicleAdv/Config/DefaultEngine.ini
--- a/Templates/TP_VehicleAdv/Config/DefaultEngine.ini
+++ b/Templates/TP_VehicleAdv/Config/DefaultEngine.ini
@@ -15,20 +15,6 @@ ServerDefaultMap=/Engine/Maps/Entry
 GlobalDefaultGameMode=/Script/TP_VehicleAdv.TP_VehicleAdvGameMode
 GlobalDefaultServerGameMode=None
 
-[/Script/AndroidFileServerEditor.AndroidFileServerRuntimeSettings]
-bEnablePlugin=True
-bAllowNetworkConnection=True
-SecurityToken=45384E7B49F38D16E0F05497235849BE
-bIncludeInShipping=False
-bAllowExternalStartInShipping=False
-bCompileAFSProject=False
-bUseCompression=False
-bLogFiles=False
-bReportStats=False
-ConnectionType=USBOnly
-bUseManualIPAddress=False
-ManualIPAddress=
-
 [/Script/Engine.RendererSettings]
 r.ReflectionMethod=1
 r.GenerateMeshDistanceFields=True
diff --git a/Templates/TP_VehicleAdvBP/Config/DefaultEngine.ini b/Templates/TP_VehicleAdvBP/Config/DefaultEngine.ini
--- a/Templates/TP_VehicleAdvBP/Config/DefaultEngine.ini
+++ b/Templates/TP_VehicleAdvBP/Config/DefaultEngine.ini
@@ -18,20 +18,6 @@ ServerDefaultMap=/Engine/Maps/Entry.Entry
 GlobalDefaultGameMode=/Game/VehicleTemplate/Blueprints/VehicleAdvGameMode.VehicleAdvGameMode_C
 GlobalDefaultServerGameMode=None
 
-[/Script/AndroidFileServerEditor.AndroidFileServerRuntimeSettings]
-bEnablePlugin=True
-bAllowNetworkConnection=True
-SecurityToken=3BE0C6B84B81CE26B8E97BBB4D1902A2
-bIncludeInShipping=False
-bAllowExternalStartInShipping=False
-bCompileAFSProject=False
-bUseCompression=False
-bLogFiles=False
-bReportStats=False
-ConnectionType=USBOnly
-bUseManualIPAddress=False
-ManualIPAddress=
-
 [/Script/Engine.RendererSettings]
 r.ReflectionMethod=1
 r.GenerateMeshDistanceFields=True
