Remove RayTracing hidden setting as it's no longer valid since the "Ray Tracing" checkbox has been removed.

diff --git a/Templates/TP_SIM_Blank/Config/TemplateDefs.ini b/Templates/TP_SIM_Blank/Config/TemplateDefs.ini
--- a/Templates/TP_SIM_Blank/Config/TemplateDefs.ini
+++ b/Templates/TP_SIM_Blank/Config/TemplateDefs.ini
@@ -103,7 +103,6 @@ ReplacementsInFiles=(Extensions=("cpp","h","ini","cs"), From="%TEMPLATENAME%", T
 ;HiddenSettings=Languages
 HiddenSettings=HardwareTarget
 HiddenSettings=GraphicsPreset
-HiddenSettings=RayTracing
 HiddenSettings=StarterContent 
 HiddenSettings=XR
 
diff --git a/Templates/TP_SIM_BlankBP/Config/TemplateDefs.ini b/Templates/TP_SIM_BlankBP/Config/TemplateDefs.ini
--- a/Templates/TP_SIM_BlankBP/Config/TemplateDefs.ini
+++ b/Templates/TP_SIM_BlankBP/Config/TemplateDefs.ini
@@ -103,7 +103,6 @@ ReplacementsInFiles=(Extensions=("cpp","h","ini","cs"), From="%TEMPLATENAME%", T
 ;HiddenSettings=Languages
 HiddenSettings=HardwareTarget
 HiddenSettings=GraphicsPreset
-HiddenSettings=RayTracing
 HiddenSettings=StarterContent 
 HiddenSettings=XR
 
diff --git a/Templates/TP_SIM_Sandbox/Config/TemplateDefs.ini b/Templates/TP_SIM_Sandbox/Config/TemplateDefs.ini
--- a/Templates/TP_SIM_Sandbox/Config/TemplateDefs.ini
+++ b/Templates/TP_SIM_Sandbox/Config/TemplateDefs.ini
@@ -102,7 +102,6 @@ ReplacementsInFiles=(Extensions=("cpp","h","ini","cs"), From="%TEMPLATENAME%", T
 ;HiddenSettings=Languages
 HiddenSettings=HardwareTarget
 HiddenSettings=GraphicsPreset
-HiddenSettings=RayTracing
 HiddenSettings=StarterContent 
 HiddenSettings=XR
 
diff --git a/Templates/TP_SIM_SandboxBP/Config/TemplateDefs.ini b/Templates/TP_SIM_SandboxBP/Config/TemplateDefs.ini
--- a/Templates/TP_SIM_SandboxBP/Config/TemplateDefs.ini
+++ b/Templates/TP_SIM_SandboxBP/Config/TemplateDefs.ini
@@ -102,7 +102,6 @@ ReplacementsInFiles=(Extensions=("cpp","h","ini","cs"), From="%TEMPLATENAME%", T
 ;HiddenSettings=Languages
 HiddenSettings=HardwareTarget
 HiddenSettings=GraphicsPreset
-HiddenSettings=RayTracing
 HiddenSettings=StarterContent 
 HiddenSettings=XR
 
diff --git a/Templates/TP_VirtualRealityBP/Config/TemplateDefs.ini b/Templates/TP_VirtualRealityBP/Config/TemplateDefs.ini
--- a/Templates/TP_VirtualRealityBP/Config/TemplateDefs.ini
+++ b/Templates/TP_VirtualRealityBP/Config/TemplateDefs.ini
@@ -52,7 +52,6 @@ FocusAsset=/Game/VRTemplate/Blueprints
 ; This project uses a thumbnail as the project icon
 bThumbnailAsIcon=true
 
-HiddenSettings=RayTracing
 HiddenSettings=Languages
 HiddenSettings=HardwareTarget
 HiddenSettings=GraphicsPreset
