Templates - Removed extra console keybind on SIM templates.

#rnx

diff --git a/Templates/TP_SIM_Blank/Config/DefaultInput.ini b/Templates/TP_SIM_Blank/Config/DefaultInput.ini
--- a/Templates/TP_SIM_Blank/Config/DefaultInput.ini
+++ b/Templates/TP_SIM_Blank/Config/DefaultInput.ini
@@ -83,6 +83,4 @@ DefaultPlayerInputClass=/Script/EnhancedInput.EnhancedPlayerInput
 DefaultInputComponentClass=/Script/EnhancedInput.EnhancedInputComponent
 DefaultTouchInterface=/Engine/MobileResources/HUD/DefaultVirtualJoysticks.DefaultVirtualJoysticks
 -ConsoleKeys=Tilde
-+ConsoleKeys=Tilde
-+ConsoleKeys=Caret
-
++ConsoleKeys=Tilde
\ No newline at end of file
diff --git a/Templates/TP_SIM_BlankBP/Config/DefaultInput.ini b/Templates/TP_SIM_BlankBP/Config/DefaultInput.ini
--- a/Templates/TP_SIM_BlankBP/Config/DefaultInput.ini
+++ b/Templates/TP_SIM_BlankBP/Config/DefaultInput.ini
@@ -84,5 +84,3 @@ DefaultInputComponentClass=/Script/EnhancedInput.EnhancedInputComponent
 DefaultTouchInterface=/Engine/MobileResources/HUD/DefaultVirtualJoysticks.DefaultVirtualJoysticks
 -ConsoleKeys=Tilde
 +ConsoleKeys=Tilde
-+ConsoleKeys=Caret
-
diff --git a/Templates/TP_SIM_Sandbox/Config/DefaultInput.ini b/Templates/TP_SIM_Sandbox/Config/DefaultInput.ini
--- a/Templates/TP_SIM_Sandbox/Config/DefaultInput.ini
+++ b/Templates/TP_SIM_Sandbox/Config/DefaultInput.ini
@@ -83,6 +83,4 @@ DefaultPlayerInputClass=/Script/EnhancedInput.EnhancedPlayerInput
 DefaultInputComponentClass=/Script/EnhancedInput.EnhancedInputComponent
 DefaultTouchInterface=/Engine/MobileResources/HUD/DefaultVirtualJoysticks.DefaultVirtualJoysticks
 -ConsoleKeys=Tilde
-+ConsoleKeys=Tilde
-+ConsoleKeys=Caret
-
++ConsoleKeys=Tilde
\ No newline at end of file
diff --git a/Templates/TP_SIM_SandboxBP/Config/DefaultInput.ini b/Templates/TP_SIM_SandboxBP/Config/DefaultInput.ini
--- a/Templates/TP_SIM_SandboxBP/Config/DefaultInput.ini
+++ b/Templates/TP_SIM_SandboxBP/Config/DefaultInput.ini
@@ -83,6 +83,4 @@ DefaultPlayerInputClass=/Script/EnhancedInput.EnhancedPlayerInput
 DefaultInputComponentClass=/Script/EnhancedInput.EnhancedInputComponent
 DefaultTouchInterface=/Engine/MobileResources/HUD/DefaultVirtualJoysticks.DefaultVirtualJoysticks
 -ConsoleKeys=Tilde
-+ConsoleKeys=Tilde
-+ConsoleKeys=Caret
-
++ConsoleKeys=Tilde
\ No newline at end of file
diff --git a/Templates/TP_TopDownBP/Config/DefaultInput.ini b/Templates/TP_TopDownBP/Config/DefaultInput.ini
--- a/Templates/TP_TopDownBP/Config/DefaultInput.ini
+++ b/Templates/TP_TopDownBP/Config/DefaultInput.ini
@@ -80,6 +80,4 @@ DefaultPlayerInputClass=/Script/EnhancedInput.EnhancedPlayerInput
 DefaultInputComponentClass=/Script/EnhancedInput.EnhancedInputComponent
 DefaultTouchInterface=None
 -ConsoleKeys=Tilde
-+ConsoleKeys=Tilde
-+ConsoleKeys=Caret
-
++ConsoleKeys=Tilde
\ No newline at end of file
