Added EyeAdaptation.CachedLightingExposure override to Simulation templates to avoid warning message due to Physically accurate light units

#Synced-CL 52478980

diff --git a/Templates/TP_SIM_Blank/Config/DefaultEngine.ini b/Templates/TP_SIM_Blank/Config/DefaultEngine.ini
--- a/Templates/TP_SIM_Blank/Config/DefaultEngine.ini
+++ b/Templates/TP_SIM_Blank/Config/DefaultEngine.ini
@@ -7,6 +7,7 @@ GameDefaultMap=/Game/Main
 
 [/Script/Engine.RendererSettings]
 r.DefaultFeature.AutoExposure.ExtendDefaultLuminanceRange=True
+r.EyeAdaptation.CachedLightingPreExposure=8
 r.AllowStaticLighting=False
 r.Shadow.Virtual.Enable=1
 r.GenerateMeshDistanceFields=True
diff --git a/Templates/TP_SIM_BlankBP/Config/DefaultEngine.ini b/Templates/TP_SIM_BlankBP/Config/DefaultEngine.ini
--- a/Templates/TP_SIM_BlankBP/Config/DefaultEngine.ini
+++ b/Templates/TP_SIM_BlankBP/Config/DefaultEngine.ini
@@ -9,6 +9,7 @@ GameDefaultMap=/Game/SimBlank/Levels/SimBlank.SimBlank
 
 [/Script/Engine.RendererSettings]
 r.DefaultFeature.AutoExposure.ExtendDefaultLuminanceRange=True
+r.EyeAdaptation.CachedLightingPreExposure=8
 r.AllowStaticLighting=False
 r.Shadow.Virtual.Enable=1
 r.GenerateMeshDistanceFields=True
diff --git a/Templates/TP_SIM_Sandbox/Config/DefaultEngine.ini b/Templates/TP_SIM_Sandbox/Config/DefaultEngine.ini
--- a/Templates/TP_SIM_Sandbox/Config/DefaultEngine.ini
+++ b/Templates/TP_SIM_Sandbox/Config/DefaultEngine.ini
@@ -9,6 +9,7 @@ GameDefaultMap=/Game/SimSandbox/Levels/SimSandbox.SimSandbox
 
 [/Script/Engine.RendererSettings]
 r.DefaultFeature.AutoExposure.ExtendDefaultLuminanceRange=True
+r.EyeAdaptation.CachedLightingPreExposure=8
 r.AllowStaticLighting=False
 r.Shadow.Virtual.Enable=1
 r.GenerateMeshDistanceFields=True
diff --git a/Templates/TP_SIM_SandboxBP/Config/DefaultEngine.ini b/Templates/TP_SIM_SandboxBP/Config/DefaultEngine.ini
--- a/Templates/TP_SIM_SandboxBP/Config/DefaultEngine.ini
+++ b/Templates/TP_SIM_SandboxBP/Config/DefaultEngine.ini
@@ -9,6 +9,7 @@ GameDefaultMap=/Game/SimSandbox/Levels/SimSandbox.SimSandbox
 
 [/Script/Engine.RendererSettings]
 r.DefaultFeature.AutoExposure.ExtendDefaultLuminanceRange=True
+r.EyeAdaptation.CachedLightingPreExposure=8
 r.AllowStaticLighting=False
 r.Shadow.Virtual.Enable=1
 r.GenerateMeshDistanceFields=True
